Set Ctrl-C handler in SlateViewer.

diff --git a/Engine/Source/Programs/SlateViewer/Private/Linux/SlateViewerMainLinux.cpp b/Engine/Source/Programs/SlateViewer/Private/Linux/SlateViewerMainLinux.cpp
--- a/Engine/Source/Programs/SlateViewer/Private/Linux/SlateViewerMainLinux.cpp
+++ b/Engine/Source/Programs/SlateViewer/Private/Linux/SlateViewerMainLinux.cpp
@@ -10,6 +10,8 @@ static FString GSavedCommandLine;
  */
 int main( int argc, char** argv )
 {
+	FPlatformMisc::SetGracefulTerminationHandler();
+
 	for (int32 Option = 1; Option < argc; Option++)
 	{
 		GSavedCommandLine += TEXT(" ");
