Source code access now works in SlateViewer and Launcher

Made sure to manually load the appropriate modules for source code access.
The hot reload dependency that VisualStudioSourceCodeAccess had was bringing in the whole editor as a dependency, so this is now restricted to editor builds.
SlateViewer now compiles in CoreUObject as the source code accessor module requires it for its settings.

TTP# 333675 - CHECKIN: SLATE: WidgetReflector code reference hyperlinks don't work anymore

diff --git a/Engine/Plugins/Developer/VisualStudioSourceCodeAccess/Source/VisualStudioSourceCodeAccess/Private/VisualStudioSourceCodeAccessor.cpp b/Engine/Plugins/Developer/VisualStudioSourceCodeAccess/Source/VisualStudioSourceCodeAccess/Private/VisualStudioSourceCodeAccessor.cpp
--- a/Engine/Plugins/Developer/VisualStudioSourceCodeAccess/Source/VisualStudioSourceCodeAccess/Private/VisualStudioSourceCodeAccessor.cpp
+++ b/Engine/Plugins/Developer/VisualStudioSourceCodeAccess/Source/VisualStudioSourceCodeAccess/Private/VisualStudioSourceCodeAccessor.cpp
@@ -6,8 +6,10 @@
 #include "ISourceCodeAccessModule.h"
 #include "ModuleManager.h"
 #include "DesktopPlatformModule.h"
-#include "Developer/HotReload/Public/IHotReload.h"
 
+#if WITH_EDITOR
+#include "Developer/HotReload/Public/IHotReload.h"
+#endif
 
 #include "AllowWindowsPlatformTypes.h"
 #if VSACCESSOR_HAS_DTE
@@ -67,8 +69,10 @@ void FVisualStudioSourceCodeAccessor::Startup()
 {
 	VSLaunchTime = 0.0;
 
+#if WITH_EDITOR
 	// Setup compilation for saving all VS documents upon compilation start
 	IHotReloadModule::Get().OnModuleCompilerStarted().AddStatic( &SaveVisualStudioDocuments );
+#endif
 
 	// Cache this so we don't have to do it on a background thread
 	GetSolutionPath();
@@ -83,11 +87,13 @@ void FVisualStudioSourceCodeAccessor::Startup()
 
 void FVisualStudioSourceCodeAccessor::Shutdown()
 {
+#if WITH_EDITOR
 	// Unregister the hot-reload callback
 	if(IHotReloadModule::IsAvailable())
 	{
 		IHotReloadModule::Get().OnModuleCompilerStarted().RemoveStatic( &SaveVisualStudioDocuments );
 	}
+#endif
 }
 
 #if VSACCESSOR_HAS_DTE
diff --git a/Engine/Plugins/Developer/VisualStudioSourceCodeAccess/Source/VisualStudioSourceCodeAccess/VisualStudioSourceCodeAccess.Build.cs b/Engine/Plugins/Developer/VisualStudioSourceCodeAccess/Source/VisualStudioSourceCodeAccess/VisualStudioSourceCodeAccess.Build.cs
--- a/Engine/Plugins/Developer/VisualStudioSourceCodeAccess/Source/VisualStudioSourceCodeAccess/VisualStudioSourceCodeAccess.Build.cs
+++ b/Engine/Plugins/Developer/VisualStudioSourceCodeAccess/Source/VisualStudioSourceCodeAccess/VisualStudioSourceCodeAccess.Build.cs
@@ -12,10 +12,14 @@ namespace UnrealBuildTool.Rules
 					"Core",
 					"SourceCodeAccess",
 					"DesktopPlatform",
-					"HotReload"
 				}
 			);
 
+			if (UEBuildConfiguration.bBuildEditor)
+			{
+				PrivateDependencyModuleNames.Add("HotReload");
+			}
+
 			if (WindowsPlatform.bHasVisualStudioDTE)
 			{
 				// This module requires atlbase.h to be included before Windows headers, so we can make use of shared PCHs.  This
diff --git a/Engine/Source/Programs/SlateViewer/Private/SlateViewerApp.cpp b/Engine/Source/Programs/SlateViewer/Private/SlateViewerApp.cpp
--- a/Engine/Source/Programs/SlateViewer/Private/SlateViewerApp.cpp
+++ b/Engine/Source/Programs/SlateViewer/Private/SlateViewerApp.cpp
@@ -3,7 +3,7 @@
 #include "SlateViewerApp.h"
 #include "RequiredProgramMainCPPInclude.h"
 #include "STestSuite.h"
-
+#include "ISourceCodeAccessModule.h"
 
 IMPLEMENT_APPLICATION(SlateViewer, "SlateViewer");
 
@@ -22,6 +22,18 @@ void RunSlateViewer( const TCHAR* CommandLine )
 	// crank up a normal Slate application using the platform's standalone renderer
 	FSlateApplication::InitializeAsStandaloneApplication(GetStandardStandaloneRenderer());
 
+	// Load the source code access module
+	ISourceCodeAccessModule& SourceCodeAccessModule = FModuleManager::LoadModuleChecked<ISourceCodeAccessModule>( FName( "SourceCodeAccess" ) );
+	
+	// Manually load in the source code access plugins, as standalone programs don't currently support plugins.
+#if PLATFORM_MAC
+	IModuleInterface& XCodeSourceCodeAccessModule = FModuleManager::LoadModuleChecked<IModuleInterface>( FName( "XCodeSourceCodeAccess" ) );
+	SourceCodeAccessModule.SetAccessor(FName("XCodeSourceCodeAccess"));
+#elif PLATFORM_WINDOWS
+	IModuleInterface& VisualStudioSourceCodeAccessModule = FModuleManager::LoadModuleChecked<IModuleInterface>( FName( "VisualStudioSourceCodeAccess" ) );
+	SourceCodeAccessModule.SetAccessor(FName("VisualStudioSourceCodeAccess"));
+#endif
+
 	// set the application name
 	FGlobalTabmanager::Get()->SetApplicationTitle(NSLOCTEXT("SlateViewer", "AppTitle", "Slate Viewer"));
 	FModuleManager::LoadModuleChecked<ISlateReflectorModule>("SlateReflector").RegisterTabSpawner(WorkspaceMenu::DeveloperMenu);
diff --git a/Engine/Source/Programs/SlateViewer/SlateViewer.Build.cs b/Engine/Source/Programs/SlateViewer/SlateViewer.Build.cs
--- a/Engine/Source/Programs/SlateViewer/SlateViewer.Build.cs
+++ b/Engine/Source/Programs/SlateViewer/SlateViewer.Build.cs
@@ -17,9 +17,19 @@ public class SlateViewer : ModuleRules
 				"SlateCore",
 				"SlateReflector",
 				"StandaloneRenderer",
+				"SourceCodeAccess",
 			}
 		);
 
+		if (Target.Platform == UnrealTargetPlatform.Mac)
+		{
+			PrivateDependencyModuleNames.Add("XCodeSourceCodeAccess");
+		}
+		else if (Target.Platform == UnrealTargetPlatform.Win64)
+		{
+			PrivateDependencyModuleNames.Add("VisualStudioSourceCodeAccess");
+		}
+
 		PrivateIncludePaths.Add("Runtime/Launch/Private");		// For LaunchEngineLoop.cpp include
 
 		if (Target.Platform == UnrealTargetPlatform.IOS)
diff --git a/Engine/Source/Programs/SlateViewer/SlateViewer.Target.cs b/Engine/Source/Programs/SlateViewer/SlateViewer.Target.cs
--- a/Engine/Source/Programs/SlateViewer/SlateViewer.Target.cs
+++ b/Engine/Source/Programs/SlateViewer/SlateViewer.Target.cs
@@ -57,7 +57,9 @@ public class SlateViewerTarget : TargetRules
 
 		// SlateViewer doesn't ever compile with the engine linked in
 		UEBuildConfiguration.bCompileAgainstEngine = false;
-		UEBuildConfiguration.bCompileAgainstCoreUObject = false;
+
+		// We need CoreUObject compiled in as the source code access module requires it
+		UEBuildConfiguration.bCompileAgainstCoreUObject = true;
 
 		// SlateViewer.exe has no exports, so no need to verify that a .lib and .exp file was emitted by
 		// the linker.
