SlateViewer - Added UEBuildConfiguration.bCompileWithStatsWithoutEngine = true, it's safer to enable it per-app

diff --git a/Engine/Source/Programs/SlateViewer/SlateViewer.Target.cs b/Engine/Source/Programs/SlateViewer/SlateViewer.Target.cs
--- a/Engine/Source/Programs/SlateViewer/SlateViewer.Target.cs
+++ b/Engine/Source/Programs/SlateViewer/SlateViewer.Target.cs
@@ -61,6 +61,8 @@ public class SlateViewerTarget : TargetRules
 		// We need CoreUObject compiled in as the source code access module requires it
 		UEBuildConfiguration.bCompileAgainstCoreUObject = true;
 
+		UEBuildConfiguration.bCompileWithStatsWithoutEngine = true;
+
 		// SlateViewer.exe has no exports, so no need to verify that a .lib and .exp file was emitted by
 		// the linker.
 		OutLinkEnvironmentConfiguration.bHasExports = false;
