Removed debug printf from SlateViewerMainMac.cpp

diff --git a/Engine/Source/Programs/SlateViewer/Private/Mac/SlateViewerMainMac.cpp b/Engine/Source/Programs/SlateViewer/Private/Mac/SlateViewerMainMac.cpp
--- a/Engine/Source/Programs/SlateViewer/Private/Mac/SlateViewerMainMac.cpp
+++ b/Engine/Source/Programs/SlateViewer/Private/Mac/SlateViewerMainMac.cpp
@@ -78,7 +78,6 @@ static FString GSavedCommandLine;
 
 - (void)sendEvent:(NSEvent*)event
 {
-	printf("sending event to CEF, type %d\n", (int32)[event type]);
 	CefScopedSendingEvent sendingEventScoper;
 	[super sendEvent:event];
 }
