Enable opening new windows from SlateViewer SWebBrowser tab.

diff --git a/Engine/Source/Programs/SlateViewer/Private/SlateViewerApp.cpp b/Engine/Source/Programs/SlateViewer/Private/SlateViewerApp.cpp
--- a/Engine/Source/Programs/SlateViewer/Private/SlateViewerApp.cpp
+++ b/Engine/Source/Programs/SlateViewer/Private/SlateViewerApp.cpp
@@ -7,6 +7,8 @@
 #include "SPerfSuite.h"
 #include "SDockTab.h"
 #include "SWebBrowser.h"
+#include "IWebBrowserWindow.h"
+#include "IWebBrowserPopupFeatures.h"
 
 IMPLEMENT_APPLICATION(SlateViewer, "SlateViewer");
 
@@ -76,6 +78,141 @@ void RunSlateViewer( const TCHAR* CommandLine )
 	FSlateApplication::Shutdown();
 }
 
+namespace
+{
+
+	TMap<TWeakPtr<IWebBrowserWindow>, TWeakPtr<SWindow>> BrowserWindowWidgets;
+
+	bool HandleBeforePopup(FString Url, FString Target)
+	{
+		return false; // Allow any popup
+	}
+
+	bool HandleBrowserCloseWindow(const TWeakPtr<IWebBrowserWindow>& BrowserWindowPtr)
+	{
+		TSharedPtr<IWebBrowserWindow> BrowserWindow = BrowserWindowPtr.Pin();
+		if(BrowserWindow.IsValid())
+		{
+			if(!BrowserWindow->IsClosing())
+			{
+				// If the browser is not set to close, we tell the browser to close which will call back into this handler function.
+				BrowserWindow->CloseBrowser(false);
+			}
+			else
+			{
+				const TWeakPtr<SWindow>* FoundBrowserWindowWidgetPtr = BrowserWindowWidgets.Find(BrowserWindow);
+				if(FoundBrowserWindowWidgetPtr != nullptr)
+				{
+					TSharedPtr<SWindow> FoundBrowserWindowWidget = FoundBrowserWindowWidgetPtr->Pin();
+					if(FoundBrowserWindowWidget.IsValid())
+					{
+						FoundBrowserWindowWidget->RequestDestroyWindow();
+					}
+					BrowserWindowWidgets.Remove(BrowserWindow);
+					return true;
+				}
+			}
+		}
+
+		return false;
+	}
+
+	bool HandleBrowserCreateWindow(const TWeakPtr<IWebBrowserWindow>& NewBrowserWindow, const TWeakPtr<IWebBrowserPopupFeatures>& PopupFeatures, TSharedPtr<SWindow> ParentWindow)
+	{
+		if(ParentWindow.IsValid())
+		{
+			TSharedPtr<IWebBrowserPopupFeatures> PopupFeaturesSP = PopupFeatures.Pin();
+			check(PopupFeatures.IsValid())
+
+			const int PosX = PopupFeaturesSP->IsXSet() ? PopupFeaturesSP->GetX() : 100;
+			const int PosY = PopupFeaturesSP->IsYSet() ? PopupFeaturesSP->GetY() : 100;
+			const FVector2D BrowserWindowPosition(PosX, PosY);
+
+			const int Width = PopupFeaturesSP->IsWidthSet() ? PopupFeaturesSP->GetWidth() : 800;
+			const int Height = PopupFeaturesSP->IsHeightSet() ? PopupFeaturesSP->GetHeight() : 600;
+			const FVector2D BrowserWindowSize(Width, Height);
+
+			const ESizingRule::Type SizeingRule = PopupFeaturesSP->IsResizable() ? ESizingRule::UserSized : ESizingRule::FixedSize;
+
+			TSharedPtr<IWebBrowserWindow> NewBrowserWindowSP = NewBrowserWindow.Pin();
+			check(NewBrowserWindowSP.IsValid())
+
+			TSharedRef<SWindow> BrowserWindowWidget =
+				SNew(SWindow)
+				.Title(LOCTEXT("WebBrowserWindow_Title", "Web Browser"))
+				.ClientSize(BrowserWindowSize)
+				.ScreenPosition(BrowserWindowPosition)
+				.AutoCenter(EAutoCenter::None)
+				.SizingRule(SizeingRule)
+				.SupportsMaximize(SizeingRule != ESizingRule::FixedSize)
+				.SupportsMinimize(SizeingRule != ESizingRule::FixedSize)
+				.HasCloseButton(true)
+				.IsInitiallyMaximized(PopupFeaturesSP->IsFullscreen())
+				.LayoutBorder(FMargin(0));
+
+			// Setup browser widget.
+			TSharedPtr<SWebBrowser> BrowserWidget;
+			{
+				TSharedPtr<SVerticalBox> Contents;
+				BrowserWindowWidget->SetContent(
+					SNew(SBorder)
+					.VAlign(VAlign_Fill)
+					.HAlign(HAlign_Fill)
+					.Padding(0)
+					[
+						SAssignNew(Contents, SVerticalBox)
+					]);
+
+				Contents->AddSlot()
+					[
+						SAssignNew(BrowserWidget, SWebBrowser, NewBrowserWindowSP)
+						.ShowControls(PopupFeaturesSP->IsToolBarVisible())
+						.ShowAddressBar(PopupFeaturesSP->IsLocationBarVisible())
+						.OnCreateWindow(FOnCreateWindowDelegate::CreateStatic(&HandleBrowserCreateWindow, ParentWindow))
+						.OnCloseWindow(FOnCloseWindowDelegate::CreateStatic(&HandleBrowserCloseWindow))
+					];
+			}
+
+			// Setup some OnClose stuff.
+			{
+				struct FLocal
+				{
+					static void RequestDestroyWindowOverride(const TSharedRef<SWindow>& Window, TWeakPtr<IWebBrowserWindow> BrowserWindowPtr)
+					{
+						TSharedPtr<IWebBrowserWindow> BrowserWindow = BrowserWindowPtr.Pin();
+						if(BrowserWindow.IsValid())
+						{
+							if(BrowserWindow->IsClosing())
+							{
+								FSlateApplicationBase::Get().RequestDestroyWindow(Window);
+							}
+							else
+							{
+								// Notify the browser window that we would like to close it.  On the CEF side, this will 
+								//  result in a call to FWebBrowserHandler::DoClose only if the JavaScript onbeforeunload
+								//  event handler allows it.
+								BrowserWindow->CloseBrowser(false);
+							}
+						}
+					}
+				};
+
+				BrowserWindowWidget->SetRequestDestroyWindowOverride(FRequestDestroyWindowOverride::CreateStatic(&FLocal::RequestDestroyWindowOverride, TWeakPtr<IWebBrowserWindow>(NewBrowserWindow)));
+			}
+
+			FSlateApplication::Get().AddWindowAsNativeChild(BrowserWindowWidget, ParentWindow.ToSharedRef());
+			BrowserWindowWidget->BringToFront();
+			FSlateApplication::Get().SetKeyboardFocus( BrowserWidget, EKeyboardFocusCause::SetDirectly );
+
+			BrowserWindowWidgets.Add(NewBrowserWindow, BrowserWindowWidget);
+			return true;
+		}
+		return false;
+	}
+
+
+}
+
 TSharedRef<SDockTab> SpawnWebBrowserTab(const FSpawnTabArgs& Args)
 {
 	return SNew(SDockTab)
@@ -84,6 +221,8 @@ TSharedRef<SDockTab> SpawnWebBrowserTab(const FSpawnTabArgs& Args)
 		.TabRole(ETabRole::NomadTab)
 		[
 			SNew(SWebBrowser)
+			.OnCreateWindow(FOnCreateWindowDelegate::CreateStatic(&HandleBrowserCreateWindow, Args.GetOwnerWindow()))
+			.OnCloseWindow(FOnCloseWindowDelegate::CreateStatic(&HandleBrowserCloseWindow))
 			.ParentWindow(Args.GetOwnerWindow())
 		];
 }
