SlateViewer: Unload modules at shutdown.
Ensures that CEF is uninitialized and has no threads hanging around, so the process will exit.

diff --git a/Engine/Source/Programs/SlateViewer/Private/SlateViewerApp.cpp b/Engine/Source/Programs/SlateViewer/Private/SlateViewerApp.cpp
--- a/Engine/Source/Programs/SlateViewer/Private/SlateViewerApp.cpp
+++ b/Engine/Source/Programs/SlateViewer/Private/SlateViewerApp.cpp
@@ -74,7 +74,7 @@ void RunSlateViewer( const TCHAR* CommandLine )
 		FSlateApplication::Get().Tick();		
 		FPlatformProcess::Sleep(0);
 	}
-
+	FModuleManager::Get().UnloadModulesAtShutdown();
 	FSlateApplication::Shutdown();
 }
 
