SlateViewer - uses an FTaskTagScope

diff --git a/Engine/Source/Programs/SlateViewer/Private/SlateViewerApp.cpp b/Engine/Source/Programs/SlateViewer/Private/SlateViewerApp.cpp
--- a/Engine/Source/Programs/SlateViewer/Private/SlateViewerApp.cpp
+++ b/Engine/Source/Programs/SlateViewer/Private/SlateViewerApp.cpp
@@ -3,6 +3,7 @@
 #include "SlateViewerApp.h"
 #include "RequiredProgramMainCPPInclude.h"
 #include "Widgets/Testing/STestSuite.h"
+#include "Widgets/Testing/SStarshipSuite.h"
 #include "ISourceCodeAccessModule.h"
 #include "Widgets/Testing/SPerfSuite.h"
 #include "Widgets/Docking/SDockTab.h"
@@ -11,6 +12,9 @@
 #include "IWebBrowserWindow.h"
 #include "IWebBrowserPopupFeatures.h"
 
+#include "Styling/StarshipCoreStyle.h"
+
+
 IMPLEMENT_APPLICATION(SlateViewer, "SlateViewer");
 
 #define LOCTEXT_NAMESPACE "SlateViewer"
@@ -23,6 +27,9 @@ namespace WorkspaceMenu
 
 int RunSlateViewer( const TCHAR* CommandLine )
 {
+
+	FTaskTagScope TaskTagScope(ETaskTag::EGameThread);	
+
 	// start up the main loop
 	GEngineLoop.PreInit(CommandLine);
 
@@ -32,8 +39,12 @@ int RunSlateViewer( const TCHAR* CommandLine )
 	// Tell the module manager it may now process newly-loaded UObjects when new C++ modules are loaded
 	FModuleManager::Get().StartProcessingNewlyLoadedObjects();
 
+
 	// crank up a normal Slate application using the platform's standalone renderer
 	FSlateApplication::InitializeAsStandaloneApplication(GetStandardStandaloneRenderer());
+	FSlateApplication::Get().EnableMenuAnimations(false);
+
+	FSlateApplication::InitHighDPI(true);
 
 	// Load the source code access module
 	ISourceCodeAccessModule& SourceCodeAccessModule = FModuleManager::LoadModuleChecked<ISourceCodeAccessModule>( FName( "SourceCodeAccess" ) );
@@ -59,6 +70,12 @@ int RunSlateViewer( const TCHAR* CommandLine )
 		// Bring up perf test
 		SummonPerfTestSuite();
 	}
+	else if (FCoreStyle::IsStarshipStyle())
+	{
+		FGlobalTabmanager::Get()->SetApplicationTitle(LOCTEXT("AppTitle", "Starship Slate Viewer"));
+		FAppStyle::SetAppStyleSetName(FStarshipCoreStyle::GetCoreStyle().GetStyleSetName());
+		RestoreStarshipSuite();	
+	}
 	else
 	{
 		// Bring up the test suite.
@@ -79,13 +96,17 @@ int RunSlateViewer( const TCHAR* CommandLine )
 		FTicker::GetCoreTicker().Tick(FApp::GetDeltaTime());
 		FSlateApplication::Get().PumpMessages();
 		FSlateApplication::Get().Tick();		
-		FPlatformProcess::Sleep(0);
+		FPlatformProcess::Sleep(0.01);
+	
+		GFrameCounter++;
 	}
 
 	FCoreDelegates::OnExit.Broadcast();
 	FSlateApplication::Shutdown();
 	FModuleManager::Get().UnloadModulesAtShutdown();
 
+	GEngineLoop.AppPreExit();
+	GEngineLoop.AppExit();
 
 	return 0;
 }
