- Added WebModule initialization to SlateViewer so the web browser window works properly

diff --git a/Engine/Source/Programs/SlateViewer/Private/SlateViewerApp.cpp b/Engine/Source/Programs/SlateViewer/Private/SlateViewerApp.cpp
--- a/Engine/Source/Programs/SlateViewer/Private/SlateViewerApp.cpp
+++ b/Engine/Source/Programs/SlateViewer/Private/SlateViewerApp.cpp
@@ -11,6 +11,7 @@
 #include "Framework/Application/SlateApplication.h"
 #include "IWebBrowserWindow.h"
 #include "IWebBrowserPopupFeatures.h"
+#include "WebBrowserModule.h"
 
 #include "Styling/StarshipCoreStyle.h"
 
@@ -87,6 +88,11 @@ int RunSlateViewer( const TCHAR* CommandLine )
 	SharedPointerTesting::TestSharedPointer<ESPMode::ThreadSafe>();
 #endif
 
+
+	// initialize the web browser
+	FWebBrowserInitSettings BrowserInitSettings = FWebBrowserInitSettings();
+	IWebBrowserModule::Get().CustomInitialize(BrowserInitSettings);
+	
 	// loop while the server does the rest
 	while (!IsEngineExitRequested())
 	{
